CIS Fix - Missing include for non unity build

#rnx

diff --git a/Engine/Plugins/Runtime/AR/Microsoft/AzureSpatialAnchors/Source/AzureSpatialAnchors/Private/AzureSpatialAnchorsFunctionLibrary.cpp b/Engine/Plugins/Runtime/AR/Microsoft/AzureSpatialAnchors/Source/AzureSpatialAnchors/Private/AzureSpatialAnchorsFunctionLibrary.cpp
--- a/Engine/Plugins/Runtime/AR/Microsoft/AzureSpatialAnchors/Source/AzureSpatialAnchors/Private/AzureSpatialAnchorsFunctionLibrary.cpp
+++ b/Engine/Plugins/Runtime/AR/Microsoft/AzureSpatialAnchors/Source/AzureSpatialAnchors/Private/AzureSpatialAnchorsFunctionLibrary.cpp
@@ -2,6 +2,7 @@
 
 #include "AzureSpatialAnchorsFunctionLibrary.h"
 #include "Engine/Engine.h"
+#include "GameFramework/WorldSettings.h"
 #include "IAzureSpatialAnchors.h"
 
 #include "ARPin.h"
