AsureSpatialAnchors build break
-Added task include

diff --git a/Engine/Plugins/Runtime/AR/Microsoft/AzureSpatialAnchors/Source/AzureSpatialAnchors/Public/IAzureSpatialAnchors.h b/Engine/Plugins/Runtime/AR/Microsoft/AzureSpatialAnchors/Source/AzureSpatialAnchors/Public/IAzureSpatialAnchors.h
--- a/Engine/Plugins/Runtime/AR/Microsoft/AzureSpatialAnchors/Source/AzureSpatialAnchors/Public/IAzureSpatialAnchors.h
+++ b/Engine/Plugins/Runtime/AR/Microsoft/AzureSpatialAnchors/Source/AzureSpatialAnchors/Public/IAzureSpatialAnchors.h
@@ -12,6 +12,7 @@
 #include "Features/IModularFeature.h"
 #include "Delegates/Delegate.h"
 #include "Delegates/DelegateCombinations.h"
+#include "Async/TaskGraphInterfaces.h"
 
 class AZURESPATIALANCHORS_API IAzureSpatialAnchors : public IModularFeature
 {
