UE4 -> UE reference change in comment.

diff --git a/Engine/Plugins/Runtime/AR/AzureSpatialAnchors/Source/AzureSpatialAnchors/Public/AzureCloudSpatialAnchor.h b/Engine/Plugins/Runtime/AR/AzureSpatialAnchors/Source/AzureSpatialAnchors/Public/AzureCloudSpatialAnchor.h
--- a/Engine/Plugins/Runtime/AR/AzureSpatialAnchors/Source/AzureSpatialAnchors/Public/AzureCloudSpatialAnchor.h
+++ b/Engine/Plugins/Runtime/AR/AzureSpatialAnchors/Source/AzureSpatialAnchors/Public/AzureCloudSpatialAnchor.h
@@ -14,8 +14,8 @@ class AZURESPATIALANCHORS_API UAzureCloudSpatialAnchor : public UObject
 	
 public:
 	/**
-	 * The UE4 Cloud Anchor ID.
-	 * This maps to a CloudSpatialAnchor within the microsoft api.  It is unique per UE4 app instance.
+	 * The UE Cloud Anchor ID.
+	 * This maps to a CloudSpatialAnchor within the microsoft api.  It is unique per UE app instance.
 	 * It is different from the CloudSpatialAnchorIdentifier, which uniquely identifies an anchor in the Azure cloud.
 	 * This is not exposed to blueprint, blueprint should use the UAzureCloudSpatialAnchor object as a whole.
 	 */
