UE4 -> UE reference change in comments.

diff --git a/Engine/Plugins/Runtime/AR/AzureSpatialAnchors/Source/AzureSpatialAnchors/Public/AzureSpatialAnchorsFunctionLibrary.h b/Engine/Plugins/Runtime/AR/AzureSpatialAnchors/Source/AzureSpatialAnchors/Public/AzureSpatialAnchorsFunctionLibrary.h
--- a/Engine/Plugins/Runtime/AR/AzureSpatialAnchors/Source/AzureSpatialAnchors/Public/AzureSpatialAnchorsFunctionLibrary.h
+++ b/Engine/Plugins/Runtime/AR/AzureSpatialAnchors/Source/AzureSpatialAnchors/Public/AzureSpatialAnchorsFunctionLibrary.h
@@ -160,7 +160,7 @@ public:
 	// * @param CloudIdentifier				The Azure Cloud Spatial Anchor Identifier of the cloud anchor we will try to load.
 	// * @param PinId						Specify the name of the Pin to load into.  Will fail if a pin of this name already exists.  If empty an auto-generated id will be used.
 	// * @param OutARPin					Filled in with the pin created, if successful.
-	// * @param OutAzureCloudSpatialAnchor	Filled in with the UE4 representation of the cloud spatial anchor created, if successful.
+	// * @param OutAzureCloudSpatialAnchor	Filled in with the UE representation of the cloud spatial anchor created, if successful.
 	// * @param OutResult					The Result enumeration.
 	// * @param OutErrorString				Additional informatiuon about the OutResult (often empty).
 	// */
